A generic recording device that is implementable by python plugins.  This is the stubbed functionality and only provides an abstract class for implementors.

#rnx
#Synced-CL 51338485

diff --git a/Engine/Plugins/Animation/LiveLinkGenericRecordingDevice/LiveLinkGenericRecordingDevice.uplugin b/Engine/Plugins/Animation/LiveLinkGenericRecordingDevice/LiveLinkGenericRecordingDevice.uplugin
new file mode 100644
--- /dev/null
+++ b/Engine/Plugins/Animation/LiveLinkGenericRecordingDevice/LiveLinkGenericRecordingDevice.uplugin
@@ -0,0 +1,31 @@
+{
+	"FileVersion": 3,
+	"Version": 1,
+	"VersionName": "1.0",
+	"FriendlyName": "LiveLinkGenericRecordingDevice",
+	"Description": "Provides a generic recording device that can be used by Python scripts for Live Link Hub. It is hidden by default as the classes are abstract.",
+	"Category": "Other",
+	"CreatedBy": "Epic Games, Inc.",
+	"CreatedByURL": "https://epicgames.com",
+	"DocsURL": "",
+	"MarketplaceURL": "",
+	"SupportURL": "",
+	"CanContainContent": true,
+	"IsBetaVersion": false,
+    "Hidden": true,
+	"IsExperimentalVersion": true,
+	"Installed": false,
+	"Modules": [
+		{
+			"Name": "LiveLinkGenericRecordingDevice",
+			"Type": "Editor",
+			"LoadingPhase": "Default"
+		}
+	],
+	"Plugins": [
+		{
+			"Name": "LiveLinkDevice",
+			"Enabled": true
+		}
+	]
+}
diff --git a/Engine/Plugins/Animation/LiveLinkGenericRecordingDevice/Source/LiveLinkGenericRecordingDevice/LiveLinkGenericRecordingDevice.Build.cs b/Engine/Plugins/Animation/LiveLinkGenericRecordingDevice/Source/LiveLinkGenericRecordingDevice/LiveLinkGenericRecordingDevice.Build.cs
new file mode 100644
--- /dev/null
+++ b/Engine/Plugins/Animation/LiveLinkGenericRecordingDevice/Source/LiveLinkGenericRecordingDevice/LiveLinkGenericRecordingDevice.Build.cs
@@ -0,0 +1,49 @@
+// Copyright Epic Games, Inc. All Rights Reserved.
+ 
+using UnrealBuildTool;
+
+public class LiveLinkGenericRecordingDevice : ModuleRules
+{
+	public LiveLinkGenericRecordingDevice(ReadOnlyTargetRules Target) : base(Target)
+	{
+		PCHUsage = ModuleRules.PCHUsageMode.UseExplicitOrSharedPCHs;
+
+		PublicIncludePaths.AddRange(
+			new string[] {
+				// Add public include paths here
+			}
+		);
+
+		PrivateIncludePaths.AddRange(
+			new string[] {
+				// Add private include paths here
+			}
+		);
+
+		PublicDependencyModuleNames.AddRange(
+			new string[]
+			{
+				"Core",
+			}
+		);
+
+		PrivateDependencyModuleNames.AddRange(
+			new string[]
+			{
+				"CoreUObject",
+				"Engine",
+				"Slate",
+				"SlateCore",
+				"LiveLinkDevice",
+				"Projects"
+			}
+		);
+
+		DynamicallyLoadedModuleNames.AddRange(
+			new string[]
+			{
+				// Add dynamically loaded modules here
+			}
+		);
+	}
+}
diff --git a/Engine/Plugins/Animation/LiveLinkGenericRecordingDevice/Source/LiveLinkGenericRecordingDevice/Private/LiveLinkGenericRecordingDevice.cpp b/Engine/Plugins/Animation/LiveLinkGenericRecordingDevice/Source/LiveLinkGenericRecordingDevice/Private/LiveLinkGenericRecordingDevice.cpp
new file mode 100644
--- /dev/null
+++ b/Engine/Plugins/Animation/LiveLinkGenericRecordingDevice/Source/LiveLinkGenericRecordingDevice/Private/LiveLinkGenericRecordingDevice.cpp
@@ -0,0 +1,60 @@
+// Copyright Epic Games, Inc. All Rights Reserved.
+
+#include "LiveLinkGenericRecordingDevice.h"
+
+TSubclassOf<ULiveLinkDeviceSettings> ULiveLinkGenericRecordingDevice::GetSettingsClass() const
+{
+	return SettingsClass ? SettingsClass : TSubclassOf<ULiveLinkDeviceSettings>(ULiveLinkGenericRecordingDeviceSettings::StaticClass());
+}
+
+FText ULiveLinkGenericRecordingDevice::GetDisplayName() const
+{
+	// Delegate to BlueprintNativeEvent version
+	return GetDisplayNameHelper();
+}
+
+ EDeviceHealth ULiveLinkGenericRecordingDevice::GetDeviceHealth() const
+{
+	// Delegate to BlueprintNativeEvent version
+	return GetDeviceHealthHelper();
+}
+
+FText ULiveLinkGenericRecordingDevice::GetHealthText() const
+{
+	// Delegate to BlueprintNativeEvent version
+	return GetHealthTextHelper();
+}
+
+FText ULiveLinkGenericRecordingDevice::GetDisplayNameHelper_Implementation() const
+{
+	return FText::FromString(TEXT("Python Recording Device"));
+}
+
+EDeviceHealth ULiveLinkGenericRecordingDevice::GetDeviceHealthHelper_Implementation() const
+{
+	return EDeviceHealth::Nominal;
+}
+
+FText ULiveLinkGenericRecordingDevice::GetHealthTextHelper_Implementation() const
+{
+	return FText::FromString(TEXT("OK"));
+}
+
+bool ULiveLinkGenericRecordingDevice::StartRecording_Implementation()
+{
+	bIsRecording = true;
+	return true;
+}
+
+bool ULiveLinkGenericRecordingDevice::StopRecording_Implementation() 
+{
+	bIsRecording = false;
+	return true;
+}
+
+bool ULiveLinkGenericRecordingDevice::IsRecording_Implementation() const
+{
+	return bIsRecording;
+}
+
+
diff --git a/Engine/Plugins/Animation/LiveLinkGenericRecordingDevice/Source/LiveLinkGenericRecordingDevice/Private/LiveLinkGenericRecordingDevice.h b/Engine/Plugins/Animation/LiveLinkGenericRecordingDevice/Source/LiveLinkGenericRecordingDevice/Private/LiveLinkGenericRecordingDevice.h
new file mode 100644
--- /dev/null
+++ b/Engine/Plugins/Animation/LiveLinkGenericRecordingDevice/Source/LiveLinkGenericRecordingDevice/Private/LiveLinkGenericRecordingDevice.h
@@ -0,0 +1,64 @@
+// Copyright Epic Games, Inc. All Rights Reserved.
+
+#pragma once
+
+#include "LiveLinkDevice.h"
+#include "LiveLinkDeviceCapability_Connection.h"
+#include "LiveLinkDeviceCapability_Recording.h"
+#include "LiveLinkGenericRecordingDevice.generated.h"
+
+UCLASS(Blueprintable)
+class ULiveLinkGenericRecordingDeviceSettings : public ULiveLinkDeviceSettings
+{
+    GENERATED_BODY()
+};
+
+UCLASS(Abstract, Blueprintable)
+class ULiveLinkGenericRecordingDevice : public ULiveLinkDevice
+    , public ILiveLinkDeviceCapability_Recording  // Implement the recording interface
+{
+    GENERATED_BODY()
+
+public:
+    //~ Begin ULiveLinkDevice interface (pure virtuals)
+    // These delegate to BlueprintNativeEvent versions that Python can override
+    virtual TSubclassOf<ULiveLinkDeviceSettings> GetSettingsClass() const override;
+    virtual FText GetDisplayName() const override;
+    virtual EDeviceHealth GetDeviceHealth() const override;
+    virtual FText GetHealthText() const override;
+    //~ End ULiveLinkDevice interface
+
+    // BlueprintNativeEvent versions that Python scripts can override
+    UFUNCTION(BlueprintNativeEvent, BlueprintCallable, Category="Live Link Device")
+    FText GetDisplayNameHelper() const;
+    virtual FText GetDisplayNameHelper_Implementation() const;
+
+    UFUNCTION(BlueprintNativeEvent, BlueprintCallable, Category="Live Link Device")
+    EDeviceHealth GetDeviceHealthHelper() const;
+    virtual EDeviceHealth GetDeviceHealthHelper_Implementation() const;
+
+    UFUNCTION(BlueprintNativeEvent, BlueprintCallable, Category="Live Link Device")
+    FText GetHealthTextHelper() const;
+    virtual FText GetHealthTextHelper_Implementation() const;
+
+    //~ Begin ILiveLinkDeviceCapability_Recording interface
+    // These are already BlueprintNativeEvent so they can be overridden in Python
+    virtual bool StartRecording_Implementation() override;
+    virtual bool StopRecording_Implementation() override;
+    virtual bool IsRecording_Implementation() const override;
+    //~ End ILiveLinkDeviceCapability_Recording interface
+
+protected:
+    // Settings class to use (can be set by derived Python class)
+    UPROPERTY(EditDefaultsOnly, Category="Live Link Device")
+    TSubclassOf<ULiveLinkDeviceSettings> SettingsClass;
+
+    UPROPERTY(BlueprintReadWrite, Category="Live Link Device")
+	bool bIsRecording = false;
+
+    UPROPERTY(BlueprintReadWrite, Category="Live Link Device")
+	FString Slate;
+
+    UPROPERTY(BlueprintReadWrite, Category="Live Link Device")
+	FString Take;
+};
diff --git a/Engine/Plugins/Animation/LiveLinkGenericRecordingDevice/Source/LiveLinkGenericRecordingDevice/Private/LiveLinkGenericRecordingDeviceModule.cpp b/Engine/Plugins/Animation/LiveLinkGenericRecordingDevice/Source/LiveLinkGenericRecordingDevice/Private/LiveLinkGenericRecordingDeviceModule.cpp
new file mode 100644
--- /dev/null
+++ b/Engine/Plugins/Animation/LiveLinkGenericRecordingDevice/Source/LiveLinkGenericRecordingDevice/Private/LiveLinkGenericRecordingDeviceModule.cpp
@@ -0,0 +1,18 @@
+// Copyright Epic Games, Inc. All Rights Reserved.
+#include "LiveLinkGenericRecordingDeviceModule.h"
+
+#define LOCTEXT_NAMESPACE "FLiveLinkGenericRecordingDeviceModule"
+
+void FLiveLinkGenericRecordingDeviceModule::StartupModule()
+{
+	UE_LOGF(LogTemp, Verbose, "LiveLinkGenericRecordingDevice module started");
+}
+
+void FLiveLinkGenericRecordingDeviceModule::ShutdownModule()
+{
+	UE_LOGF(LogTemp, Verbose, "LiveLinkGenericRecordingDevice module shutdown");
+}
+
+#undef LOCTEXT_NAMESPACE
+
+IMPLEMENT_MODULE(FLiveLinkGenericRecordingDeviceModule, LiveLinkGenericRecordingDevice)
diff --git a/Engine/Plugins/Animation/LiveLinkGenericRecordingDevice/Source/LiveLinkGenericRecordingDevice/Private/LiveLinkGenericRecordingDeviceModule.h b/Engine/Plugins/Animation/LiveLinkGenericRecordingDevice/Source/LiveLinkGenericRecordingDevice/Private/LiveLinkGenericRecordingDeviceModule.h
new file mode 100644
--- /dev/null
+++ b/Engine/Plugins/Animation/LiveLinkGenericRecordingDevice/Source/LiveLinkGenericRecordingDevice/Private/LiveLinkGenericRecordingDeviceModule.h
@@ -0,0 +1,13 @@
+// Copyright Epic Games, Inc. All Rights Reserved.
+
+#pragma once
+
+#include "Modules/ModuleManager.h"
+
+class FLiveLinkGenericRecordingDeviceModule : public IModuleInterface
+{
+public:
+	/** IModuleInterface implementation */
+	virtual void StartupModule() override;
+	virtual void ShutdownModule() override;
+};
